Fix IOS location services
#ue4
#ios

diff --git a/Engine/Plugins/Runtime/LocationServicesIOSImpl/Source/LocationServicesIOSImpl/Classes/LocationServicesIOSImpl.h b/Engine/Plugins/Runtime/LocationServicesIOSImpl/Source/LocationServicesIOSImpl/Classes/LocationServicesIOSImpl.h
--- a/Engine/Plugins/Runtime/LocationServicesIOSImpl/Source/LocationServicesIOSImpl/Classes/LocationServicesIOSImpl.h
+++ b/Engine/Plugins/Runtime/LocationServicesIOSImpl/Source/LocationServicesIOSImpl/Classes/LocationServicesIOSImpl.h
@@ -46,7 +46,7 @@ public:
 	* a default-value-filled struct.
 	* @return - the last known location from updates
 	*/
-	virtual FLocationData GetLastKnownLocation() override;
+	virtual FLocationServicesData GetLastKnownLocation() override;
 
 	/**
 	* Checks if the supplied Accuracy is available on the current device.
diff --git a/Engine/Plugins/Runtime/LocationServicesIOSImpl/Source/LocationServicesIOSImpl/Private/LocationServicesIOSImpl.cpp b/Engine/Plugins/Runtime/LocationServicesIOSImpl/Source/LocationServicesIOSImpl/Private/LocationServicesIOSImpl.cpp
--- a/Engine/Plugins/Runtime/LocationServicesIOSImpl/Source/LocationServicesIOSImpl/Private/LocationServicesIOSImpl.cpp
+++ b/Engine/Plugins/Runtime/LocationServicesIOSImpl/Source/LocationServicesIOSImpl/Private/LocationServicesIOSImpl.cpp
@@ -81,7 +81,7 @@ DEFINE_LOG_CATEGORY(LogLocationServicesIOS);
         oldLocation = nil;
     }
     
-    FLocationData locationData;
+	FLocationServicesData locationData;
     locationData.Timestamp = (float)newLocation.timestamp.timeIntervalSince1970;
     locationData.Longitude = (float)newLocation.coordinate.longitude;
     locationData.Latitude = (float)newLocation.coordinate.latitude;
@@ -156,9 +156,9 @@ bool ULocationServicesIOSImpl::StopLocationService()
     return false;
 }
 
-FLocationData ULocationServicesIOSImpl::GetLastKnownLocation()
+FLocationServicesData ULocationServicesIOSImpl::GetLastKnownLocation()
 {  
-	FLocationData LocData;
+	FLocationServicesData LocData;
 
 	if (LocationDelegate != nil)
 	{
