Fix deprecation warning in iOS location services

diff --git a/Engine/Plugins/Runtime/LocationServicesIOSImpl/Source/LocationServicesIOSImpl/Private/LocationServicesIOSImpl.cpp b/Engine/Plugins/Runtime/LocationServicesIOSImpl/Source/LocationServicesIOSImpl/Private/LocationServicesIOSImpl.cpp
--- a/Engine/Plugins/Runtime/LocationServicesIOSImpl/Source/LocationServicesIOSImpl/Private/LocationServicesIOSImpl.cpp
+++ b/Engine/Plugins/Runtime/LocationServicesIOSImpl/Source/LocationServicesIOSImpl/Private/LocationServicesIOSImpl.cpp
@@ -40,8 +40,8 @@ DEFINE_LOG_CATEGORY(LogLocationServicesIOS);
      LocManager.delegate = self;
      LocManager.distanceFilter = DistanceFilter;
      LocManager.desiredAccuracy = Accuracy;
-     
-     if ([CLLocationManager authorizationStatus] == kCLAuthorizationStatusNotDetermined)
+    
+     if (LocManager.authorizationStatus == kCLAuthorizationStatusNotDetermined)
      {
         //we need to request location services
         [LocManager requestAlwaysAuthorization];
@@ -210,16 +210,21 @@ bool ULocationServicesIOSImpl::IsLocationAccuracyAvailable(ELocationAccuracy Acc
 	
 bool ULocationServicesIOSImpl::IsLocationServiceEnabled()
 {
+    CLLocationManager *LocManager = [[CLLocationManager alloc] init];;
+
 	bool bEnabled = [CLLocationManager locationServicesEnabled];
 	if (!bEnabled)
 	{
 		UE_LOG(LogLocationServicesIOS, Log, TEXT("ULocationServicesIOSImpl::IsServiceEnabled() - location services disabled in settings"));
 	}
-	bool bAuthorized = [CLLocationManager authorizationStatus] != kCLAuthorizationStatusDenied;
+    bool bAuthorized = LocManager.authorizationStatus == kCLAuthorizationStatusDenied;
 	if (!bAuthorized)
 	{
 		UE_LOG(LogLocationServicesIOS, Log, TEXT("ULocationServicesIOSImpl::IsServiceEnabled() - location services have not been authorized for use"));
 	}
 
+    [LocManager release];
+    LocManager = nil;
+
 	return bEnabled && bAuthorized;
-}
\ No newline at end of file
+}
