Fix static analyser report of missing dealloc in LocationServicesIOSImpl

#rnx

diff --git a/Engine/Plugins/Runtime/LocationServicesIOSImpl/Source/LocationServicesIOSImpl/Private/LocationServicesIOSImpl.cpp b/Engine/Plugins/Runtime/LocationServicesIOSImpl/Source/LocationServicesIOSImpl/Private/LocationServicesIOSImpl.cpp
--- a/Engine/Plugins/Runtime/LocationServicesIOSImpl/Source/LocationServicesIOSImpl/Private/LocationServicesIOSImpl.cpp
+++ b/Engine/Plugins/Runtime/LocationServicesIOSImpl/Source/LocationServicesIOSImpl/Private/LocationServicesIOSImpl.cpp
@@ -15,6 +15,8 @@ DEFINE_LOG_CATEGORY(LogLocationServicesIOS);
 -(void)startUpdatingLocation;
 -(void)stopUpdatingLocation;
 -(CLLocation*)getLastKnownLocation;
+-(void)dealloc;
+
 @property CLLocationAccuracy Accuracy;
 @property CLLocationDistance DistanceFilter;
 @property(strong, nonatomic) CLLocationManager* LocManager;
@@ -68,6 +70,16 @@ DEFINE_LOG_CATEGORY(LogLocationServicesIOS);
     
     return nil;
 }
+-(void)dealloc
+{
+	if (LocManager != nil)
+	{
+		LocManager.delegate = nil;
+		[LocManager release];
+		LocManager = nil;
+	}
+	[super dealloc];
+}
 /*
 * Callback from the LocationManager when there's an update to our location
 */
